Fix: VideoEncoderConfig defaults were incorrect, leading to 0 bitrates throughout

diff --git a/Engine/Plugins/Experimental/AVCodecs/AVCodecsCore/Source/AVCodecsCore/Public/Video/VideoEncoder.h b/Engine/Plugins/Experimental/AVCodecs/AVCodecsCore/Source/AVCodecsCore/Public/Video/VideoEncoder.h
--- a/Engine/Plugins/Experimental/AVCodecs/AVCodecsCore/Source/AVCodecsCore/Public/Video/VideoEncoder.h
+++ b/Engine/Plugins/Experimental/AVCodecs/AVCodecsCore/Source/AVCodecsCore/Public/Video/VideoEncoder.h
@@ -155,9 +155,10 @@ public:
 
 	uint32 TargetFramerate = 60;
 
-	int32 TargetBitrate = 0;
-	int32 MaxBitrate = 0;
-	int32 MinBitrate = 0;
+	// Individual coder implementations check for -1 bitrate and assign value
+	int32 TargetBitrate = -1;
+	int32 MaxBitrate = -1;
+	int32 MinBitrate = -1;
 	// Advanced bitrate settings. Used for situations such as simulcast / SVC
 	TOptional<int32> Bitrates[Video::MaxSpatialLayers][Video::MaxTemporalStreams];
 
diff --git a/Engine/Plugins/Experimental/AVCodecs/LibVpxCodecs/Source/LibVpxCodecs/Public/Video/Encoders/VideoEncoderLibVpxVP9.hpp b/Engine/Plugins/Experimental/AVCodecs/LibVpxCodecs/Source/LibVpxCodecs/Public/Video/Encoders/VideoEncoderLibVpxVP9.hpp
--- a/Engine/Plugins/Experimental/AVCodecs/LibVpxCodecs/Source/LibVpxCodecs/Public/Video/Encoders/VideoEncoderLibVpxVP9.hpp
+++ b/Engine/Plugins/Experimental/AVCodecs/LibVpxCodecs/Source/LibVpxCodecs/Public/Video/Encoders/VideoEncoderLibVpxVP9.hpp
@@ -769,7 +769,7 @@ FAVResult TVideoEncoderLibVpxVP9<TResource>::InitAndSetControlSettings(FVideoEnc
 			SvcParams->loopfilter_ctrl[si] = PerformanceFlagsBySpatialIndex[si].DeblockMode;
 		}
 		bool bDenoiserOn = Config.bDenoisingOn && PerformanceFlagsBySpatialIndex[NumSpatialLayers - 1].bAllowDenoising;
-		SAFECONTROLVP9(Encoder.Get(), VP9E_SET_NOISE_SENSITIVITY, bDenoiserOn ? 1 : 0);
+		::vpx_codec_control(Encoder.Get(), VP9E_SET_NOISE_SENSITIVITY, bDenoiserOn ? 1 : 0);
 	}
 
 	SAFECONTROLVP9(Encoder.Get(), VP8E_SET_MAX_INTRA_BITRATE_PCT, RCMaxIntraTarget);
@@ -847,7 +847,7 @@ FAVResult TVideoEncoderLibVpxVP9<TResource>::InitAndSetControlSettings(FVideoEnc
 
 	if (!PerformanceFlags.bUsePerLayerSpeed)
 	{
-		SAFECONTROLVP9(Encoder.Get(), VP9E_SET_NOISE_SENSITIVITY, Config.bDenoisingOn ? 1 : 0);
+		::vpx_codec_control(Encoder.Get(), VP9E_SET_NOISE_SENSITIVITY, Config.bDenoisingOn ? 1 : 0);
 	}
 
 	// Enable encoder skip of static/low content blocks.
@@ -1685,7 +1685,7 @@ bool TVideoEncoderLibVpxVP9<TResource>::SetSvcRates(FVideoEncoderConfigLibVpx co
 	if (bSeenActiveLayer && PerformanceFlags.bUsePerLayerSpeed)
 	{
 		bool bDenoiserOn = this->AppliedConfig.bDenoisingOn && PerformanceFlagsBySpatialIndex[NumActiveSpatialLayers - 1].bAllowDenoising;
-		SAFECONTROLVP9(Encoder.Get(), VP9E_SET_NOISE_SENSITIVITY, bDenoiserOn ? 1 : 0);
+		::vpx_codec_control(Encoder.Get(), VP9E_SET_NOISE_SENSITIVITY, bDenoiserOn ? 1 : 0);
 	}
 
 	if (bHigherLayersEnabled && !bForceKeyFrame)
