[AVCodecs] Fix: VP9 intermittently working. Set your default values kiddos

diff --git a/Engine/Plugins/Experimental/AVCodecs/LibVpxCodecs/Source/LibVpxCodecs/Public/Video/Encoders/VideoEncoderLibVpxVP9.h b/Engine/Plugins/Experimental/AVCodecs/LibVpxCodecs/Source/LibVpxCodecs/Public/Video/Encoders/VideoEncoderLibVpxVP9.h
--- a/Engine/Plugins/Experimental/AVCodecs/LibVpxCodecs/Source/LibVpxCodecs/Public/Video/Encoders/VideoEncoderLibVpxVP9.h
+++ b/Engine/Plugins/Experimental/AVCodecs/LibVpxCodecs/Source/LibVpxCodecs/Public/Video/Encoders/VideoEncoderLibVpxVP9.h
@@ -35,7 +35,7 @@ public:
 
 	virtual FAVResult ApplyConfig() override;
 
-	virtual FAVResult SendFrame(TSharedPtr<FVideoResourceCPU> const& Resource, uint32 Timestamp, bool bForceKeyframe = false) override;
+	virtual FAVResult SendFrame(TSharedPtr<FVideoResourceCPU> const& Resource, uint32 Timestamp, bool bTriggerKeyFrame = false) override;
 
 	virtual FAVResult ReceivePacket(FVideoPacket& OutPacket) override;
 
@@ -148,11 +148,11 @@ private:
 		}
 
 	public:
-		EFrameType	   FrameType;
-		int32		   Width;
-		int32		   Height;
-		uint32		   Timestamp;
-		int32		   QP;
+		EFrameType	   FrameType = EFrameType::P;
+		int32		   Width = 0;
+		int32		   Height = 0;
+		uint32		   Timestamp = 0;
+		int32		   QP = -1;
 		TOptional<int> SpatialIndex;
 		TOptional<int> TemporalIndex;
 
@@ -167,29 +167,29 @@ private:
 	};
 
 private:
-	EProfile				 Profile;
-	int64					 Timestamp = 0;
-	bool					 bForceKeyFrame = false;
-	size_t					 PicsSinceKey = 0;
-	uint8					 NumTemporalLayers;		 // Number of configured TLs
-	uint8					 NumSpatialLayers;		 // Number of configured SLs
-	uint8					 NumActiveSpatialLayers; // Number of actively encoded SLs
-	uint8					 FirstActiveLayer;
-	bool					 bIsSvc;
-	bool					 bIsFlexibleMode;
-	EInterLayerPrediction InterLayerPrediction;
-	bool					 bExternalRefControl;
-	bool					 bFullSuperframeDrop;
-	bool					 bLayerBuffering;
-	bool					 bFirstFrameInPicture;
-	uint32					 RCMaxIntraTarget;
-	FEncodedImage			 EncodedImage;
-	FGroupOfFramesInfo		 Gof;
-	FVideoBitrateAllocation	 CurrentBitrateAllocation;
-	bool					 bSsInfoNeeded;
-	bool					 bForceAllActiveLayers;
-	bool					 bVpxConfigChanged;
-	FCodecSpecificInfo		 CodecSpecific;
+	EProfile				Profile = EProfile::Profile0;
+	int64					Timestamp = 0;
+	bool					bForceKeyFrame = true;
+	size_t					PicsSinceKey = 0;
+	uint8					NumTemporalLayers = 0;		// Number of configured TLs
+	uint8					NumSpatialLayers = 0;		// Number of configured SLs
+	uint8					NumActiveSpatialLayers = 0; // Number of actively encoded SLs
+	uint8					FirstActiveLayer = 0;
+	bool					bIsSvc = false;
+	bool					bIsFlexibleMode = false;
+	EInterLayerPrediction	InterLayerPrediction = EInterLayerPrediction::On;
+	bool					bExternalRefControl = false;
+	bool					bFullSuperframeDrop = true;
+	bool					bLayerBuffering = false;
+	bool					bFirstFrameInPicture = true;
+	uint32					RCMaxIntraTarget = 0;
+	FEncodedImage			EncodedImage;
+	FGroupOfFramesInfo		Gof;
+	FVideoBitrateAllocation CurrentBitrateAllocation;
+	bool					bSsInfoNeeded = false;
+	bool					bForceAllActiveLayers = false;
+	bool					bVpxConfigChanged = true;
+	FCodecSpecificInfo		CodecSpecific;
 
 	// Performance flags, ordered by `min_pixel_count`.
 	const FPerformanceFlags PerformanceFlags = GetDefaultPerformanceFlags();
@@ -200,15 +200,15 @@ private:
 	TArray<FRefFrameBuffer>								RefBuf;
 	TArray<FScalableVideoController::FLayerFrameConfig> LayerFrames;
 
-	TUniquePtr<FScalableVideoController> SvcController;
-	TUniquePtr<FInputImage>				 InputImage;
+	TUniquePtr<FScalableVideoController> SvcController = nullptr;
+	TUniquePtr<FInputImage>				 InputImage = nullptr;
 
 private:
 	TUniquePtr<vpx_codec_ctx_t, LibVpxUtil::FCodecContextDeleter> Encoder = nullptr;
-	TUniquePtr<vpx_codec_enc_cfg_t>  VpxConfig = nullptr;
+	TUniquePtr<vpx_codec_enc_cfg_t>								  VpxConfig = nullptr;
 	TUniquePtr<vpx_image_t, LibVpxUtil::FImageDeleter>			  RawImage = nullptr;
-	TUniquePtr<vpx_svc_extra_cfg_t>  SvcParams = nullptr;
-	TUniquePtr<vpx_svc_frame_drop_t> SvcDropFrame = nullptr;
+	TUniquePtr<vpx_svc_extra_cfg_t>								  SvcParams = nullptr;
+	TUniquePtr<vpx_svc_frame_drop_t>							  SvcDropFrame = nullptr;
 
 public:
 	void GetEncodedLayerFrame(const vpx_codec_cx_pkt* Packet);
@@ -216,9 +216,10 @@ public:
 
 namespace Internal
 {
+	template <typename TResource>
 	inline void EncoderOutputCodedPacketCallback(vpx_codec_cx_pkt* Packet, void* UserData)
 	{
-		static_cast<TVideoEncoderLibVpxVP9<FVideoResourceCPU>*>(UserData)->GetEncodedLayerFrame(Packet);
+		static_cast<TVideoEncoderLibVpxVP9<TResource>*>(UserData)->GetEncodedLayerFrame(Packet);
 	}
 } // namespace Internal
 
diff --git a/Engine/Plugins/Experimental/AVCodecs/LibVpxCodecs/Source/LibVpxCodecs/Public/Video/Encoders/VideoEncoderLibVpxVP9.hpp b/Engine/Plugins/Experimental/AVCodecs/LibVpxCodecs/Source/LibVpxCodecs/Public/Video/Encoders/VideoEncoderLibVpxVP9.hpp
--- a/Engine/Plugins/Experimental/AVCodecs/LibVpxCodecs/Source/LibVpxCodecs/Public/Video/Encoders/VideoEncoderLibVpxVP9.hpp
+++ b/Engine/Plugins/Experimental/AVCodecs/LibVpxCodecs/Source/LibVpxCodecs/Public/Video/Encoders/VideoEncoderLibVpxVP9.hpp
@@ -150,19 +150,19 @@ FAVResult TVideoEncoderLibVpxVP9<TResource>::ApplyConfig()
 				if (!Encoder.IsValid())
 				{
 					Encoder = TUniquePtr<vpx_codec_ctx_t, LibVpxUtil::FCodecContextDeleter>(new vpx_codec_ctx_t, LibVpxUtil::FCodecContextDeleter());
-					memset(Encoder.Get(), 0, sizeof(*(Encoder.Get())));
+					memset(Encoder.Get(), 0, sizeof(vpx_codec_ctx_t));
 				}
 
 				if (!VpxConfig.IsValid())
 				{
 					VpxConfig.Reset(new vpx_codec_enc_cfg_t);
-					memset(VpxConfig.Get(), 0, sizeof(*(VpxConfig.Get())));
+					memset(VpxConfig.Get(), 0, sizeof(vpx_codec_enc_cfg_t));
 				}
 
 				if (!SvcParams.IsValid())
 				{
 					SvcParams.Reset(new vpx_svc_extra_cfg_t);
-					memset(SvcParams.Get(), 0, sizeof(*(SvcParams.Get())));
+					memset(SvcParams.Get(), 0, sizeof(vpx_svc_extra_cfg_t));
 				}
 
 				Timestamp = 0;
@@ -339,7 +339,7 @@ FAVResult TVideoEncoderLibVpxVP9<TResource>::ApplyConfig()
 }
 
 template <typename TResource>
-FAVResult TVideoEncoderLibVpxVP9<TResource>::SendFrame(TSharedPtr<FVideoResourceCPU> const& Resource, uint32 InTimestamp, bool bForceKeyframe)
+FAVResult TVideoEncoderLibVpxVP9<TResource>::SendFrame(TSharedPtr<FVideoResourceCPU> const& Resource, uint32 InTimestamp, bool bTriggerKeyFrame)
 {
 	if (!IsOpen())
 	{
@@ -360,6 +360,11 @@ FAVResult TVideoEncoderLibVpxVP9<TResource>::SendFrame(TSharedPtr<FVideoResource
 			return EAVResult::Success;
 		}
 
+		if (bTriggerKeyFrame)
+		{
+			bForceKeyFrame = true;
+		}
+
 		if (PicsSinceKey + 1 == static_cast<size_t>(this->AppliedConfig.KeyframeInterval))
 		{
 			bForceKeyFrame = true;
@@ -810,7 +815,7 @@ FAVResult TVideoEncoderLibVpxVP9<TResource>::InitAndSetControlSettings(FVideoEnc
 			SvcDropFrame.Reset(new vpx_svc_frame_drop_t);
 		}
 
-		memset(SvcDropFrame.Get(), 0, sizeof(*(SvcDropFrame.Get())));
+		memset(SvcDropFrame.Get(), 0, sizeof(vpx_svc_frame_drop_t));
 
 		// Configure encoder to drop entire superframe whenever it needs to drop
 		// a layer. This mode is preferred over per-layer dropping which causes
@@ -830,7 +835,7 @@ FAVResult TVideoEncoderLibVpxVP9<TResource>::InitAndSetControlSettings(FVideoEnc
 
 	// Register callback for getting each spatial layer.
 	vpx_codec_priv_output_cx_pkt_cb_pair_t cbp = {
-		Internal::EncoderOutputCodedPacketCallback,
+		Internal::EncoderOutputCodedPacketCallback<TResource>,
 		reinterpret_cast<void*>(this)
 	};
 
@@ -1481,7 +1486,7 @@ vpx_svc_ref_frame_config_t TVideoEncoderLibVpxVP9<TResource>::SetReferences(bool
 	check(Gof.NumFramesInGof <= Video::MaxTemporalStreams);
 
 	vpx_svc_ref_frame_config_t RefConfig;
-	memset(&RefConfig, 0, sizeof(RefConfig));
+	memset(&RefConfig, 0, sizeof(vpx_svc_ref_frame_config_t));
 
 	const size_t   NumTemporalRefs = FMath::Max(1, NumTemporalLayers - 1);
 	const bool	   bIsInterLayerPredAllowed = InterLayerPrediction == EInterLayerPrediction::On || (InterLayerPrediction == EInterLayerPrediction::OnKeyPicture && bIsKeyPic);
