[AVCodecs] Only encode VP9 video with an even dimension. This prevents crash in windows and linux.

Problem is internally in libvpx, valid_ref_frame_size checks the dimensions. When A p-frame is encoded, the UV size from a odd frame size causes this check to fail. Two function pointers for scaling are then not set. These function pointers then get called which causes the crash.

Set bExternalRefControl to true, which also prevents the crash on windows but not linux.

diff --git a/Engine/Plugins/Experimental/AVCodecs/LibVpxCodecs/Source/LibVpxCodecs/Public/Video/Encoders/VideoEncoderLibVpxVP9.hpp b/Engine/Plugins/Experimental/AVCodecs/LibVpxCodecs/Source/LibVpxCodecs/Public/Video/Encoders/VideoEncoderLibVpxVP9.hpp
--- a/Engine/Plugins/Experimental/AVCodecs/LibVpxCodecs/Source/LibVpxCodecs/Public/Video/Encoders/VideoEncoderLibVpxVP9.hpp
+++ b/Engine/Plugins/Experimental/AVCodecs/LibVpxCodecs/Source/LibVpxCodecs/Public/Video/Encoders/VideoEncoderLibVpxVP9.hpp
@@ -93,7 +93,12 @@ FAVResult TVideoEncoderLibVpxVP9<TResource>::ApplyConfig()
 {
 	if (IsOpen())
 	{
-		FVideoEncoderConfigLibVpx const& PendingConfig = this->GetPendingConfig();
+		// TODO (Eden.Harris) RTCP-7171 Change this to const reference and GetPendingConfig when no longer changing the width/height
+		// VP9 currently requires even width and height otherwise it can crash when encoding P-Frames.
+		// Round size down to even number. By rounding down, it saves copying the frame to a larger buffer.
+		FVideoEncoderConfigLibVpx& PendingConfig = this->EditPendingConfig();
+		PendingConfig.Width = PendingConfig.Width & (~1);
+		PendingConfig.Height = PendingConfig.Height & (~1);
 		if (this->AppliedConfig != PendingConfig)
 		{
 			if (IsInitialized())
@@ -235,9 +240,6 @@ FAVResult TVideoEncoderLibVpxVP9<TResource>::ApplyConfig()
 				// Creating a wrapper to the image - setting image data to nullptr. Actual
 				// pointer will be set in encode. Setting align to 1, as it is meaningless
 				// (actual memory is not allocated).
-				// Creating a wrapper to the image - setting image data to NULL.
-				// Actual pointer will be set in encode. Setting align to 1, as it
-				// is meaningless (no memory allocation is done here).
 				RawImage = TUniquePtr<vpx_image_t, LibVpxUtil::FImageDeleter>(::vpx_img_wrap(nullptr, PixelFormat, PendingConfig.Width, PendingConfig.Height, 1, NULL), LibVpxUtil::FImageDeleter());
 				RawImage->bit_depth = BitsForStorage;
 
@@ -277,7 +279,7 @@ FAVResult TVideoEncoderLibVpxVP9<TResource>::ApplyConfig()
 
 				bIsFlexibleMode = PendingConfig.bFlexibleMode;
 
-				bExternalRefControl = InterLayerPrediction == EInterLayerPrediction::On;
+				bExternalRefControl = true;
 
 				if (NumTemporalLayers == 1)
 				{
@@ -495,8 +497,9 @@ FAVResult TVideoEncoderLibVpxVP9<TResource>::SendFrame(TSharedPtr<FVideoResource
 			bVpxConfigChanged = false;
 		}
 
-		check(Resource->GetWidth() == RawImage->d_w);
-		check(Resource->GetHeight() == RawImage->d_h);
+		// TODO (Eden.Harris) Remove round down after RTCP-7171
+		check((Resource->GetWidth() & ~1) == RawImage->d_w);
+		check((Resource->GetHeight() & ~1) == RawImage->d_h);
 
 		InputImage = MakeUnique<FInputImage>(InTimestamp);
 
