[AVCodecs] Fix Pixel Streaming 2 VP8 crash when frame resizes.

When the frame is resized and the bitrate changes, the VP8 encoder was not destroyed and recreated resulting in the encoder trying to encoder a frame with incorrect size. This most likely caused the encoder to access memory that was invalid.

This changes forces the encoder to be recreated when the frame size, keyframe interval, pixel format or other settings changes. In addition, if the target bitrate, min or max quality or framerate changes, it updates the encoder config rather than forcing complete recreation.

diff --git a/Engine/Plugins/Experimental/AVCodecs/LibVpxCodecs/Source/LibVpxCodecs/Public/Video/Encoders/VideoEncoderLibVpxVP8.hpp b/Engine/Plugins/Experimental/AVCodecs/LibVpxCodecs/Source/LibVpxCodecs/Public/Video/Encoders/VideoEncoderLibVpxVP8.hpp
--- a/Engine/Plugins/Experimental/AVCodecs/LibVpxCodecs/Source/LibVpxCodecs/Public/Video/Encoders/VideoEncoderLibVpxVP8.hpp
+++ b/Engine/Plugins/Experimental/AVCodecs/LibVpxCodecs/Source/LibVpxCodecs/Public/Video/Encoders/VideoEncoderLibVpxVP8.hpp
@@ -67,17 +67,33 @@ FAVResult TVideoEncoderLibVpxVP8<TResource>::ApplyConfig()
 		{
 			if (IsInitialized())
 			{
-				if (this->AppliedConfig.TargetBitrate == PendingConfig.TargetBitrate)
+				// Check if a setting has changed and destroy the encoder for re-creation.
+				// If TargetBitrate, MinQuality, MaxQuality or Framerate has changed then only need to change the config.
+				if (this->AppliedConfig.Width != PendingConfig.Width
+					|| this->AppliedConfig.Height != PendingConfig.Height
+					|| this->AppliedConfig.KeyframeInterval != PendingConfig.KeyframeInterval
+					|| this->AppliedConfig.PixelFormat != PendingConfig.PixelFormat
+					|| this->AppliedConfig.NumberOfCores != PendingConfig.NumberOfCores
+					|| this->AppliedConfig.bDenoisingOn != PendingConfig.bDenoisingOn)
 				{
-					// Bitrate has stayed the same, but something else has changed. In that case we need to destroy and re-init
+					// Setting has changed. In that case we need to destroy and re-init
 					Destroy();
 					FAVResult::Log(EAVResult::Success, TEXT("Re-initializing encoding session"), TEXT("LibVpx"));
 				}
-				else
+				else if (
+					this->AppliedConfig.TargetBitrate != PendingConfig.TargetBitrate
+					|| this->AppliedConfig.MaxQuality != PendingConfig.MaxQuality
+					|| this->AppliedConfig.MinQuality != PendingConfig.MinQuality
+					)
 				{
 					bSendStream = PendingConfig.TargetBitrate > 0;
 					VpxConfig->rc_target_bitrate = PendingConfig.TargetBitrate / 1000;
 
+					VpxConfig->rc_min_quantizer = (1.0f - (PendingConfig.MaxQuality / 100.0f)) * 63.0f;
+					VpxConfig->rc_max_quantizer = (1.0f - (PendingConfig.MinQuality / 100.0f)) * 63.0f;
+
+					RCMaxIntraTarget = MaxIntraTarget(VpxConfig->rc_buf_optimal_sz, PendingConfig.Framerate);
+
 					vpx_codec_err_t Result = ::vpx_codec_enc_config_set(Encoder.Get(), VpxConfig.Get());
 					if (Result != VPX_CODEC_OK)
 					{
@@ -85,6 +101,11 @@ FAVResult TVideoEncoderLibVpxVP8<TResource>::ApplyConfig()
 						return FAVResult(EAVResult::Error, FString::Printf(TEXT("Error configuring codec, error code: %d, details: %s"), Result, *ErrorString), TEXT("LibVpx"));
 					}
 				}
+				// Only change framerate if that is the only setting changed
+				else if (this->AppliedConfig.Framerate!= PendingConfig.Framerate)
+				{
+					RCMaxIntraTarget = MaxIntraTarget(VpxConfig->rc_buf_optimal_sz, PendingConfig.Framerate);
+				}
 			}
 
 			if (!IsInitialized())
