[AVCodecs] QOL: Update LibVpx config to use QP instead of quality. This reflects the fact the LibVpx itself expects QP

diff --git a/Engine/Plugins/Experimental/AVCodecs/LibVpxCodecs/Source/LibVpxCodecs/Private/Video/Encoders/Configs/VideoEncoderConfigLibVpx.cpp b/Engine/Plugins/Experimental/AVCodecs/LibVpxCodecs/Source/LibVpxCodecs/Private/Video/Encoders/Configs/VideoEncoderConfigLibVpx.cpp
--- a/Engine/Plugins/Experimental/AVCodecs/LibVpxCodecs/Source/LibVpxCodecs/Private/Video/Encoders/Configs/VideoEncoderConfigLibVpx.cpp
+++ b/Engine/Plugins/Experimental/AVCodecs/LibVpxCodecs/Source/LibVpxCodecs/Private/Video/Encoders/Configs/VideoEncoderConfigLibVpx.cpp
@@ -7,9 +7,10 @@
 
 REGISTER_TYPEID(FVideoEncoderConfigLibVpx);
 
-static uint32 DEFAULT_BITRATE_MIN = 100000;
-static uint32 DEFAULT_BITRATE_TARGET = 1000000;
-static uint32 DEFAULT_BITRATE_MAX = 10000000;
+static uint32	   DEFAULT_BITRATE_MIN = 100000;
+static uint32	   DEFAULT_BITRATE_TARGET = 1000000;
+static uint32	   DEFAULT_BITRATE_MAX = 10000000;
+static const float QPRANGE = 63.f; // Range used by VP8 and VP9
 
 template <>
 DLLEXPORT FAVResult FAVExtension::TransformConfig(FVideoEncoderConfigLibVpx& OutConfig, FVideoEncoderConfig const& InConfig)
@@ -21,8 +22,8 @@ DLLEXPORT FAVResult FAVExtension::TransformConfig(FVideoEncoderConfigLibVpx& Out
 	OutConfig.MinBitrate = InConfig.MinBitrate > -1 ? InConfig.MinBitrate : DEFAULT_BITRATE_MIN;
 	OutConfig.TargetBitrate = InConfig.TargetBitrate > -1 ? InConfig.TargetBitrate : DEFAULT_BITRATE_TARGET;
 	OutConfig.MaxBitrate = InConfig.MaxBitrate > -1 ? InConfig.MaxBitrate : DEFAULT_BITRATE_MAX;
-	OutConfig.MinQuality = InConfig.MinQuality;
-	OutConfig.MaxQuality = InConfig.MaxQuality;
+	OutConfig.MaxQP = static_cast<int32>(QPRANGE * (1.0f - (InConfig.MinQuality / 100.0f)));
+	OutConfig.MinQP = static_cast<int32>(QPRANGE * (1.0f - (InConfig.MaxQuality / 100.0f)));
 	OutConfig.KeyframeInterval = InConfig.KeyframeInterval > 0 ? InConfig.KeyframeInterval : 0;
 	OutConfig.ScalabilityMode = InConfig.ScalabilityMode;
 	OutConfig.NumberOfSpatialLayers = InConfig.NumberOfSpatialLayers;
@@ -54,8 +55,8 @@ DLLEXPORT FAVResult FAVExtension::TransformConfig(FVideoEncoderConfig& OutConfig
 	OutConfig.MinBitrate = InConfig.MinBitrate > -1 ? InConfig.MinBitrate : DEFAULT_BITRATE_MIN;
 	OutConfig.TargetBitrate = InConfig.TargetBitrate > -1 ? InConfig.TargetBitrate : DEFAULT_BITRATE_TARGET;
 	OutConfig.MaxBitrate = InConfig.MaxBitrate > -1 ? InConfig.MaxBitrate : DEFAULT_BITRATE_MAX;
-	OutConfig.MinQuality = InConfig.MinQuality;
-	OutConfig.MaxQuality = InConfig.MaxQuality;
+	OutConfig.MinQuality = static_cast<int32>((1.0f - (InConfig.MaxQP / QPRANGE)) * 100.0f);
+	OutConfig.MaxQuality = static_cast<int32>((1.0f - (InConfig.MinQP / QPRANGE)) * 100.0f);
 	OutConfig.KeyframeInterval = InConfig.KeyframeInterval > 0 ? InConfig.KeyframeInterval : 0;
 	OutConfig.ScalabilityMode = InConfig.ScalabilityMode;
 	OutConfig.NumberOfSpatialLayers = InConfig.NumberOfSpatialLayers;
diff --git a/Engine/Plugins/Experimental/AVCodecs/LibVpxCodecs/Source/LibVpxCodecs/Public/Video/Encoders/Configs/VideoEncoderConfigLibVpx.h b/Engine/Plugins/Experimental/AVCodecs/LibVpxCodecs/Source/LibVpxCodecs/Public/Video/Encoders/Configs/VideoEncoderConfigLibVpx.h
--- a/Engine/Plugins/Experimental/AVCodecs/LibVpxCodecs/Source/LibVpxCodecs/Public/Video/Encoders/Configs/VideoEncoderConfigLibVpx.h
+++ b/Engine/Plugins/Experimental/AVCodecs/LibVpxCodecs/Source/LibVpxCodecs/Public/Video/Encoders/Configs/VideoEncoderConfigLibVpx.h
@@ -22,8 +22,8 @@ public:
 	TOptional<int32> Bitrates[Video::MaxSpatialLayers][Video::MaxTemporalStreams];
 	uint32			 KeyframeInterval = 0;
 	EVideoFormat	 PixelFormat;
-	int32			 MinQuality;
-	int32			 MaxQuality;
+	uint32			 MinQP;
+	uint32			 MaxQP;
 
 	int32				  NumberOfCores = 0;
 	bool				  bDenoisingOn = false;
@@ -68,8 +68,8 @@ public:
 			&& this->MinBitrate == Other.MinBitrate
 			&& this->KeyframeInterval == Other.KeyframeInterval
 			&& this->PixelFormat == Other.PixelFormat
-			&& this->MinQuality == Other.MinQuality
-			&& this->MaxQuality == Other.MaxQuality
+			&& this->MinQP == Other.MinQP
+			&& this->MaxQP == Other.MaxQP
 			&& this->NumberOfCores == Other.NumberOfCores
 			&& this->bDenoisingOn == Other.bDenoisingOn
 			&& this->bAdaptiveQpMode == Other.bAdaptiveQpMode
diff --git a/Engine/Plugins/Experimental/AVCodecs/LibVpxCodecs/Source/LibVpxCodecs/Public/Video/Encoders/VideoEncoderLibVpxVP8.hpp b/Engine/Plugins/Experimental/AVCodecs/LibVpxCodecs/Source/LibVpxCodecs/Public/Video/Encoders/VideoEncoderLibVpxVP8.hpp
--- a/Engine/Plugins/Experimental/AVCodecs/LibVpxCodecs/Source/LibVpxCodecs/Public/Video/Encoders/VideoEncoderLibVpxVP8.hpp
+++ b/Engine/Plugins/Experimental/AVCodecs/LibVpxCodecs/Source/LibVpxCodecs/Public/Video/Encoders/VideoEncoderLibVpxVP8.hpp
@@ -7,14 +7,14 @@
 
 using namespace UE::AVCodecCore::VP8;
 
-#define SAFECONTROL(Encoder, Setting, Value)                                                                                                                                                   \
-	{                                                                                                                                                                                          \
-		vpx_codec_err_t Res = ::vpx_codec_control(Encoder, Setting, Value);                                                                                                                    \
-		if (Res != VPX_CODEC_OK)                                                                                                                                                               \
-		{                                                                                                                                                                                      \
-			FString ErrorString(::vpx_codec_error_detail(Encoder));                                                                                                                            \
+#define SAFECONTROL(Encoder, Setting, Value)                                                                                                                                                         \
+	{                                                                                                                                                                                                \
+		vpx_codec_err_t Res = ::vpx_codec_control(Encoder, Setting, Value);                                                                                                                          \
+		if (Res != VPX_CODEC_OK)                                                                                                                                                                     \
+		{                                                                                                                                                                                            \
+			FString ErrorString(::vpx_codec_error_detail(Encoder));                                                                                                                                  \
 			FAVResult::Log(EAVResult::Success, FString::Printf(TEXT("Error executing ::vpx_codec_control. Setting %s, Error: %d, Details: %s"), TEXT(#Setting), Res, *ErrorString), TEXT("LibVpx")); \
-		}                                                                                                                                                                                      \
+		}                                                                                                                                                                                            \
 	}
 
 template <typename TResource>
@@ -82,15 +82,14 @@ FAVResult TVideoEncoderLibVpxVP8<TResource>::ApplyConfig()
 				}
 				else if (
 					this->AppliedConfig.TargetBitrate != PendingConfig.TargetBitrate
-					|| this->AppliedConfig.MaxQuality != PendingConfig.MaxQuality
-					|| this->AppliedConfig.MinQuality != PendingConfig.MinQuality
-					)
+					|| this->AppliedConfig.MaxQP != PendingConfig.MaxQP
+					|| this->AppliedConfig.MinQP != PendingConfig.MinQP)
 				{
 					bSendStream = PendingConfig.TargetBitrate > 0;
 					VpxConfig->rc_target_bitrate = PendingConfig.TargetBitrate / 1000;
 
-					VpxConfig->rc_min_quantizer = (1.0f - (PendingConfig.MaxQuality / 100.0f)) * 63.0f;
-					VpxConfig->rc_max_quantizer = (1.0f - (PendingConfig.MinQuality / 100.0f)) * 63.0f;
+					VpxConfig->rc_min_quantizer = PendingConfig.MinQP;
+					VpxConfig->rc_max_quantizer = PendingConfig.MaxQP;
 
 					RCMaxIntraTarget = MaxIntraTarget(VpxConfig->rc_buf_optimal_sz, PendingConfig.Framerate);
 
@@ -102,7 +101,7 @@ FAVResult TVideoEncoderLibVpxVP8<TResource>::ApplyConfig()
 					}
 				}
 				// Only change framerate if that is the only setting changed
-				else if (this->AppliedConfig.Framerate!= PendingConfig.Framerate)
+				else if (this->AppliedConfig.Framerate != PendingConfig.Framerate)
 				{
 					RCMaxIntraTarget = MaxIntraTarget(VpxConfig->rc_buf_optimal_sz, PendingConfig.Framerate);
 				}
@@ -156,8 +155,8 @@ FAVResult TVideoEncoderLibVpxVP8<TResource>::ApplyConfig()
 				VpxConfig->g_pass = VPX_RC_ONE_PASS;
 				// Handle resizing outside of libvpx.
 				VpxConfig->rc_resize_allowed = 0;
-				VpxConfig->rc_min_quantizer = (1.0f - (PendingConfig.MaxQuality / 100.0f)) * 63.0f;
-				VpxConfig->rc_max_quantizer = (1.0f - (PendingConfig.MinQuality / 100.0f)) * 63.0f;
+				VpxConfig->rc_min_quantizer = PendingConfig.MinQP;
+				VpxConfig->rc_max_quantizer = PendingConfig.MaxQP;
 				VpxConfig->rc_undershoot_pct = 100;
 				VpxConfig->rc_overshoot_pct = 15;
 				VpxConfig->rc_buf_initial_sz = 500;
@@ -466,14 +465,14 @@ typename TVideoEncoderLibVpxVP8<TResource>::EEncodeResult TVideoEncoderLibVpxVP8
 	{
 		switch (Packet->kind)
 		{
-		case VPX_CODEC_CX_FRAME_PKT:
-		{
-			FMemory::Memcpy(&Buffer.GetData()[EncodedPos], Packet->data.frame.buf, Packet->data.frame.sz);
-			EncodedPos += Packet->data.frame.sz;
-			break;
-		}
-		default:
-			break;
+			case VPX_CODEC_CX_FRAME_PKT:
+			{
+				FMemory::Memcpy(&Buffer.GetData()[EncodedPos], Packet->data.frame.buf, Packet->data.frame.sz);
+				EncodedPos += Packet->data.frame.sz;
+				break;
+			}
+			default:
+				break;
 		}
 		// End of frame
 
diff --git a/Engine/Plugins/Experimental/AVCodecs/LibVpxCodecs/Source/LibVpxCodecs/Public/Video/Encoders/VideoEncoderLibVpxVP9.hpp b/Engine/Plugins/Experimental/AVCodecs/LibVpxCodecs/Source/LibVpxCodecs/Public/Video/Encoders/VideoEncoderLibVpxVP9.hpp
--- a/Engine/Plugins/Experimental/AVCodecs/LibVpxCodecs/Source/LibVpxCodecs/Public/Video/Encoders/VideoEncoderLibVpxVP9.hpp
+++ b/Engine/Plugins/Experimental/AVCodecs/LibVpxCodecs/Source/LibVpxCodecs/Public/Video/Encoders/VideoEncoderLibVpxVP9.hpp
@@ -107,8 +107,8 @@ FAVResult TVideoEncoderLibVpxVP9<TResource>::ApplyConfig()
 					&& this->AppliedConfig.Height == PendingConfig.Height
 					&& this->AppliedConfig.KeyframeInterval == PendingConfig.KeyframeInterval
 					&& this->AppliedConfig.PixelFormat == PendingConfig.PixelFormat
-					&& this->AppliedConfig.MinQuality == PendingConfig.MinQuality
-					&& this->AppliedConfig.MaxQuality == PendingConfig.MaxQuality
+					&& this->AppliedConfig.MinQP == PendingConfig.MinQP
+					&& this->AppliedConfig.MaxQP == PendingConfig.MaxQP
 					&& this->AppliedConfig.NumberOfCores == PendingConfig.NumberOfCores
 					&& this->AppliedConfig.bDenoisingOn == PendingConfig.bDenoisingOn
 					&& this->AppliedConfig.bAdaptiveQpMode == PendingConfig.bAdaptiveQpMode
@@ -257,8 +257,8 @@ FAVResult TVideoEncoderLibVpxVP9<TResource>::ApplyConfig()
 				VpxConfig->rc_dropframe_thresh = 0;
 				VpxConfig->rc_end_usage = VPX_CBR;
 				VpxConfig->g_pass = VPX_RC_ONE_PASS;
-				VpxConfig->rc_min_quantizer = (1.0f - (PendingConfig.MaxQuality / 100.0f)) * 63.0f;
-				VpxConfig->rc_max_quantizer = (1.0f - (PendingConfig.MinQuality / 100.0f)) * 63.0f;
+				VpxConfig->rc_min_quantizer = PendingConfig.MinQP;
+				VpxConfig->rc_max_quantizer = PendingConfig.MaxQP;
 				VpxConfig->rc_undershoot_pct = 50;
 				VpxConfig->rc_overshoot_pct = 50;
 				VpxConfig->rc_buf_initial_sz = 500;
