[AVCodecs] Enable odd frame width and height with VP9.

Odd frames will not crash when rc_resize_allowed is disabled. See RTCP-7994 for a explanation on why the crash can happen when rc_resize_allowed is enabled.

diff --git a/Engine/Plugins/Experimental/AVCodecs/LibVpxCodecs/Source/LibVpxCodecs/Public/Video/Encoders/VideoEncoderLibVpxVP9.hpp b/Engine/Plugins/Experimental/AVCodecs/LibVpxCodecs/Source/LibVpxCodecs/Public/Video/Encoders/VideoEncoderLibVpxVP9.hpp
--- a/Engine/Plugins/Experimental/AVCodecs/LibVpxCodecs/Source/LibVpxCodecs/Public/Video/Encoders/VideoEncoderLibVpxVP9.hpp
+++ b/Engine/Plugins/Experimental/AVCodecs/LibVpxCodecs/Source/LibVpxCodecs/Public/Video/Encoders/VideoEncoderLibVpxVP9.hpp
@@ -93,12 +93,24 @@ FAVResult TVideoEncoderLibVpxVP9<TResource>::ApplyConfig()
 {
 	if (IsOpen())
 	{
-		// TODO (Eden.Harris) RTCP-7171 Change this to const reference and GetPendingConfig when no longer changing the width/height
-		// VP9 currently requires even width and height otherwise it can crash when encoding P-Frames.
-		// Round size down to even number. By rounding down, it saves copying the frame to a larger buffer.
+		// TODO (Eden.Harris) RTCP-7994 Change this to const reference and GetPendingConfig when no longer changing the 
+		// width/height. VP9 currently requires even width and height when rc_resize_allowed is enabled, otherwise it 
+		// can crash when encoding P-Frames. 
+		// libvpx vp9_setup_scale_factors_for_frame sets several function pointers that may crash when encoding and 
+		// they are not set. vp9_setup_scale_factors_for_frame calls valid_ref_frame_size but will not set the function 
+		// pointers if valid_ref_frame_size returns false. If a frame size has an odd width/height, for instance when 
+		// rc_resize_allowed is enabled, the encoder may encode a frame at half the size. valid_ref_frame_size checks 
+		// works by doubling the frame size and comparing it against the reference size. In the case where there is a 
+		// odd frame that is 201px wide, half size frame is 100px, and the check doubles it to 200px which will be less 
+		// than the reference frame of 201px. If the check succeeds in another iteration, it will set the function 
+		// pointers for further frames and will not crash. 
 		FVideoEncoderConfigLibVpx& PendingConfig = this->EditPendingConfig();
-		PendingConfig.Width = PendingConfig.Width & (~1);
-		PendingConfig.Height = PendingConfig.Height & (~1);
+		if (PendingConfig.bAutomaticResizeOn)
+		{
+			// Round size down to even number. By rounding down, it saves copying the frame to a larger buffer.
+			PendingConfig.Width = PendingConfig.Width & (~1);
+			PendingConfig.Height = PendingConfig.Height & (~1);
+		}
 		if (this->AppliedConfig != PendingConfig)
 		{
 			if (IsInitialized())
@@ -497,10 +509,6 @@ FAVResult TVideoEncoderLibVpxVP9<TResource>::SendFrame(TSharedPtr<FVideoResource
 			bVpxConfigChanged = false;
 		}
 
-		// TODO (Eden.Harris) Remove round down after RTCP-7171
-		check((Resource->GetWidth() & ~1) == RawImage->d_w);
-		check((Resource->GetHeight() & ~1) == RawImage->d_h);
-
 		InputImage = MakeUnique<FInputImage>(InTimestamp);
 
 		switch (Profile)
